Merging //UE4/Dev-Enterprise/Engine/... to //UE4/Release-4.23/Engine/...

Updated LiveLink Camera and Light data structures

diff --git a/Engine/Plugins/Animation/LiveLink/Source/LiveLinkComponents/Private/Controllers/LiveLinkCameraController.cpp b/Engine/Plugins/Animation/LiveLink/Source/LiveLinkComponents/Private/Controllers/LiveLinkCameraController.cpp
--- a/Engine/Plugins/Animation/LiveLink/Source/LiveLinkComponents/Private/Controllers/LiveLinkCameraController.cpp
+++ b/Engine/Plugins/Animation/LiveLink/Source/LiveLinkComponents/Private/Controllers/LiveLinkCameraController.cpp
@@ -46,6 +46,10 @@ void ULiveLinkCameraController::Tick(float DeltaTime, const FLiveLinkSubjectRepr
 				if (UCineCameraComponent* CineCameraComponent = Cast<UCineCameraComponent>(CameraComponent))
 				{
 					if (StaticData->bIsFocalLengthSupported) { CineCameraComponent->CurrentFocalLength = FrameData->FocalLength; }
+					if (StaticData->bIsApertureSupported) { CineCameraComponent->CurrentAperture = FrameData->Aperture; }
+					if (StaticData->FilmBackWidth > 0.0f) { CineCameraComponent->FilmbackSettings.SensorWidth = StaticData->FilmBackWidth; }
+					if (StaticData->FilmBackHeight > 0.0f) { CineCameraComponent->FilmbackSettings.SensorHeight = StaticData->FilmBackHeight; }
+					if (StaticData->bIsFocusDistanceSupported) { CineCameraComponent->FocusSettings.ManualFocusDistance = FrameData->FocusDistance; }
 				}
 			}
 		}
diff --git a/Engine/Plugins/Animation/LiveLink/Source/LiveLinkComponents/Private/Controllers/LiveLinkLightController.cpp b/Engine/Plugins/Animation/LiveLink/Source/LiveLinkComponents/Private/Controllers/LiveLinkLightController.cpp
--- a/Engine/Plugins/Animation/LiveLink/Source/LiveLinkComponents/Private/Controllers/LiveLinkLightController.cpp
+++ b/Engine/Plugins/Animation/LiveLink/Source/LiveLinkComponents/Private/Controllers/LiveLinkLightController.cpp
@@ -7,6 +7,7 @@
 #include "Roles/LiveLinkLightTypes.h"
 
 #include "Components/LightComponent.h"
+#include "Components/SpotLightComponent.h"
 #include "Features/IModularFeatures.h"
 #include "GameFramework/Actor.h"
 
@@ -41,6 +42,20 @@ void ULiveLinkLightController::Tick(float DeltaTime, const FLiveLinkSubjectRepre
 				if (StaticData->bIsTemperatureSupported) { LightComponent->SetTemperature(FrameData->Temperature); }
 				if (StaticData->bIsIntensitySupported) { LightComponent->SetIntensity(FrameData->Intensity); }
 				if (StaticData->bIsLightColorSupported) { LightComponent->SetLightColor(FrameData->LightColor); }
+
+				if (UPointLightComponent* PointLightComponent = Cast<UPointLightComponent>(LightComponent))
+				{
+					if (StaticData->bIsAttenuationRadiusSupported) { PointLightComponent->SetAttenuationRadius(FrameData->AttenuationRadius); }
+					if (StaticData->bIsSourceRadiusSupported) { PointLightComponent->SetSourceRadius(FrameData->SourceRadius); }
+					if (StaticData->bIsSoftSourceRadiusSupported) { PointLightComponent->SetSoftSourceRadius(FrameData->SoftSourceRadius); }
+					if (StaticData->bIsSourceLenghtSupported) { PointLightComponent->SetSourceLength(FrameData->SourceLength); }
+
+					if (USpotLightComponent* SpotlightComponent = Cast<USpotLightComponent>(LightComponent))
+					{
+						if (StaticData->bIsInnerConeAngleSupported) { SpotlightComponent->SetInnerConeAngle(FrameData->InnerConeAngle); }
+						if (StaticData->bIsOuterConeAngleSupported) { SpotlightComponent->SetOuterConeAngle(FrameData->OuterConeAngle); }
+					}
+				}
 			}
 		}
 	}
diff --git a/Engine/Plugins/Animation/LiveLink/Source/LiveLinkComponents/Public/Controllers/LiveLinkLightController.h b/Engine/Plugins/Animation/LiveLink/Source/LiveLinkComponents/Public/Controllers/LiveLinkLightController.h
--- a/Engine/Plugins/Animation/LiveLink/Source/LiveLinkComponents/Public/Controllers/LiveLinkLightController.h
+++ b/Engine/Plugins/Animation/LiveLink/Source/LiveLinkComponents/Public/Controllers/LiveLinkLightController.h
@@ -9,6 +9,8 @@
 
 
 /**
+ * Controller that uses LiveLink light data to drive a light component. 
+ * UPointLightComponent and USpotLightComponent are supported for specific properties
  */
 UCLASS()
 class LIVELINKCOMPONENTS_API ULiveLinkLightController : public ULiveLinkControllerBase
diff --git a/Engine/Source/Runtime/LiveLinkInterface/Public/Roles/LiveLinkCameraTypes.h b/Engine/Source/Runtime/LiveLinkInterface/Public/Roles/LiveLinkCameraTypes.h
--- a/Engine/Source/Runtime/LiveLinkInterface/Public/Roles/LiveLinkCameraTypes.h
+++ b/Engine/Source/Runtime/LiveLinkInterface/Public/Roles/LiveLinkCameraTypes.h
@@ -21,17 +21,39 @@ struct LIVELINKINTERFACE_API FLiveLinkCameraStaticData : public FLiveLinkTransfo
 {
 	GENERATED_BODY()
 
+public:
+
+	//Whether FieldOfView in frame data can be used
 	UPROPERTY(EditAnywhere, BlueprintReadWrite, Category = "LiveLink")
 	bool bIsFieldOfViewSupported = false;
 
+	//Whether AspectRatio in frame data can be used
 	UPROPERTY(EditAnywhere, BlueprintReadWrite, Category = "LiveLink")
 	bool bIsAspectRatioSupported = false;
 
+	//Whether FocalLength in frame data can be used
 	UPROPERTY(EditAnywhere, BlueprintReadWrite, Category = "LiveLink")
 	bool bIsFocalLengthSupported = false;
 
+	//Whether ProjectionMode in frame data can be used
 	UPROPERTY(EditAnywhere, BlueprintReadWrite, Category = "LiveLink")
 	bool bIsProjectionModeSupported = false;
+
+	//Used with CinematicCamera. Values greater than 0 will be applied
+	UPROPERTY(EditAnywhere, BlueprintReadWrite, Category = "LiveLink")
+	float FilmBackWidth = -1.0f;
+
+	//Used with CinematicCamera. Values greater than 0 will be applied
+	UPROPERTY(EditAnywhere, BlueprintReadWrite, Category = "LiveLink")
+	float FilmBackHeight = -1.0f;
+
+	//Whether Aperture in frame data can be used
+	UPROPERTY(EditAnywhere, BlueprintReadWrite, Category = "LiveLink")
+	bool bIsApertureSupported = false;
+
+	//Whether FocusDistance in frame data can be used
+	UPROPERTY(EditAnywhere, BlueprintReadWrite, Category = "LiveLink")
+	bool bIsFocusDistanceSupported = false;
 };
 
 /**
@@ -42,11 +64,11 @@ struct LIVELINKINTERFACE_API FLiveLinkCameraFrameData : public FLiveLinkTransfor
 {
 	GENERATED_BODY()
 
-	// Field of View of the camera
+	// Field of View of the camera in degrees
 	UPROPERTY(EditAnywhere, BlueprintReadWrite, Category = "LiveLink", Interp)
 	float FieldOfView = 90.f;
 
-	// Aspect Ratio of the camera
+	// Aspect Ratio of the camera (Width / Heigth)
 	UPROPERTY(EditAnywhere, BlueprintReadWrite, Category = "LiveLink", Interp)
 	float AspectRatio = 1.777778f;
 
@@ -54,7 +76,15 @@ struct LIVELINKINTERFACE_API FLiveLinkCameraFrameData : public FLiveLinkTransfor
 	UPROPERTY(EditAnywhere, BlueprintReadWrite, Category = "LiveLink", Interp)
 	float FocalLength = 50.f;
 
-	// Focal length of the camera
+	// Aperture of the camera in terms of f-stop
+	UPROPERTY(EditAnywhere, BlueprintReadWrite, Category = "LiveLink", Interp)
+	float Aperture = 2.8f;
+
+	// Focus distance of the camera in cm. Works only in manual focus method
+	UPROPERTY(EditAnywhere, BlueprintReadWrite, Category = "LiveLink", Interp)
+	float FocusDistance = 100000.0f;
+
+	// ProjectionMode of the camera (Perspective, Orthographic, etc...)
 	UPROPERTY(EditAnywhere, BlueprintReadWrite, Category = "LiveLink", Interp)
 	ELiveLinkCameraProjectionMode ProjectionMode = ELiveLinkCameraProjectionMode::Perspective;
 };
diff --git a/Engine/Source/Runtime/LiveLinkInterface/Public/Roles/LiveLinkLightTypes.h b/Engine/Source/Runtime/LiveLinkInterface/Public/Roles/LiveLinkLightTypes.h
--- a/Engine/Source/Runtime/LiveLinkInterface/Public/Roles/LiveLinkLightTypes.h
+++ b/Engine/Source/Runtime/LiveLinkInterface/Public/Roles/LiveLinkLightTypes.h
@@ -13,14 +13,43 @@ struct LIVELINKINTERFACE_API FLiveLinkLightStaticData : public FLiveLinkTransfor
 {
 	GENERATED_BODY()
 
+public:
+
+	//Whether Temperature can be used in the frame data
 	UPROPERTY(EditAnywhere, BlueprintReadWrite, Category = "LiveLink")
 	bool bIsTemperatureSupported = false;
 
+	//Whether Intensity can be used in the frame data
 	UPROPERTY(EditAnywhere, BlueprintReadWrite, Category = "LiveLink")
 	bool bIsIntensitySupported = false;
 
+	//Whether LightColor can be used in the frame data
 	UPROPERTY(EditAnywhere, BlueprintReadWrite, Category = "LiveLink")
 	bool bIsLightColorSupported = false;
+
+	//Whether InnerConeAngle can be used in the frame data. Only used for spot lights
+	UPROPERTY(EditAnywhere, BlueprintReadWrite, Category = "LiveLink")
+	bool bIsInnerConeAngleSupported = false;
+
+	//Whether OuterConeAngle can be used in the frame data. Only used for spot lights
+	UPROPERTY(EditAnywhere, BlueprintReadWrite, Category = "LiveLink")
+	bool bIsOuterConeAngleSupported = false;
+
+	//Whether AttenuationRadius can be used in the frame data. Only used for spot lights
+	UPROPERTY(EditAnywhere, BlueprintReadWrite, Category = "LiveLink")
+	bool bIsAttenuationRadiusSupported = false;
+
+	//Whether SourceLength can be used in the frame data. Only used for spot lights
+	UPROPERTY(EditAnywhere, BlueprintReadWrite, Category = "LiveLink")
+	bool bIsSourceLenghtSupported = false;
+
+	//Whether SourceRadius can be used in the frame data. Only used for spot lights
+	UPROPERTY(EditAnywhere, BlueprintReadWrite, Category = "LiveLink")
+	bool bIsSourceRadiusSupported = false;
+
+	//Whether SoftSourceRadius can be used in the frame data. Only used for spot lights
+	UPROPERTY(EditAnywhere, BlueprintReadWrite, Category = "LiveLink")
+	bool bIsSoftSourceRadiusSupported = false;
 };
 
 /**
@@ -35,13 +64,37 @@ struct LIVELINKINTERFACE_API FLiveLinkLightFrameData : public FLiveLinkTransform
 	UPROPERTY(EditAnywhere, BlueprintReadWrite, Category = "LiveLink", Interp)
 	float Temperature = 6500.f;
 
-	// Total energy that the light emits.  
+	// Total energy that the light emits in lux.  
 	UPROPERTY(EditAnywhere, BlueprintReadWrite, Category = "LiveLink", Interp)
 	float Intensity = 3.1415926535897932f;
 
-	/** Filter color of the light. */
+	// Filter color of the light.
 	UPROPERTY(EditAnywhere, BlueprintReadWrite, Category = "LiveLink", Interp)
 	FColor LightColor = FColor::White;
+
+	// Inner cone angle in degrees for a Spotlight.
+	UPROPERTY(EditAnywhere, BlueprintReadWrite, Category = "LiveLink", Interp)
+	float InnerConeAngle = 0.0f;
+
+	// Outer cone angle in degrees for a Spotlight.
+	UPROPERTY(EditAnywhere, BlueprintReadWrite, Category = "LiveLink", Interp)
+	float OuterConeAngle = 44.0f;
+
+	// Light visible influence. Works for Pointlight and Spotlight.
+	UPROPERTY(EditAnywhere, BlueprintReadWrite, Category = "LiveLink", Interp)
+	float AttenuationRadius = 1000.0f;
+
+	// Radius of light source shape. Works for Pointlight and Spotlight.
+	UPROPERTY(EditAnywhere, BlueprintReadWrite, Category = "LiveLink", Interp)
+	float SourceRadius = 0.0f;
+
+	// Soft radius of light source shape. Works for Pointlight and Spotlight.
+	UPROPERTY(EditAnywhere, BlueprintReadWrite, Category = "LiveLink", Interp)
+	float SoftSourceRadius = 0.0f;
+
+	// Length of light source shape. Works for Pointlight and Spotlight.
+	UPROPERTY(EditAnywhere, BlueprintReadWrite, Category = "LiveLink", Interp)
+	float SourceLength = 0.0f;
 };
 
 /**
