Merging fix for livelink tooltip from dev-enterprise CL#7227024

diff --git a/Engine/Source/Runtime/LiveLinkInterface/Public/LiveLinkSourceSettings.h b/Engine/Source/Runtime/LiveLinkInterface/Public/LiveLinkSourceSettings.h
--- a/Engine/Source/Runtime/LiveLinkInterface/Public/LiveLinkSourceSettings.h
+++ b/Engine/Source/Runtime/LiveLinkInterface/Public/LiveLinkSourceSettings.h
@@ -15,19 +15,22 @@
 UENUM()
 enum class ELiveLinkSourceMode : uint8
 {
-	Default,				//! The source will be run in default mode.
-							//! This mode will not attempt any type of interpolation, time synchronization,
-							//! or other processing.
-
-	Interpolated,			//! The source will be run in interpolated mode.
-							//! This mode will use FLiveLinkInterpolationSettings and is most useful
-							//! when smooth animation is desired.
-
-	TimeSynchronized,		//! The source will be run in time synchronized mode.
-							//! This mode will use FLiveLinkTimeSynchronizationSettings and is most useful
-							//! when sources need to be synchronized with multiple other external inputs
-							//! (such as video or other time synchronized sources).
-							//! Don't use if the engine isn't setup with a Timecode provider.
+	//The source will be run in default mode.
+	//This mode will not attempt any type of interpolation, time synchronization,
+	//or other processing.
+	Default,				
+
+	//The source will be run in interpolated mode.
+	//This mode will use FLiveLinkInterpolationSettings and is most useful
+	//when smooth animation is desired.
+	Interpolated,		
+
+	//The source will be run in time synchronized mode.
+	//This mode will use FLiveLinkTimeSynchronizationSettings and is most useful
+	//when sources need to be synchronized with multiple other external inputs
+	//(such as video or other time synchronized sources).
+	//Don't use if the engine isn't setup with a Timecode provider.
+	TimeSynchronized,	
 };
 
 USTRUCT()
