Added UPROEPRTY markup to FLiveLinkMetaData::SceneTime

diff --git a/Engine/Source/Runtime/LiveLinkInterface/Public/LiveLinkTypes.h b/Engine/Source/Runtime/LiveLinkInterface/Public/LiveLinkTypes.h
--- a/Engine/Source/Runtime/LiveLinkInterface/Public/LiveLinkTypes.h
+++ b/Engine/Source/Runtime/LiveLinkInterface/Public/LiveLinkTypes.h
@@ -169,6 +169,7 @@ public:
 	UPROPERTY()
 	TMap<FName, FString> StringMetaData;
 
+	UPROPERTY()
 	FQualifiedFrameTime SceneTime;
 };
 
