- Fixing crash when interpolation mode is selected on livelink subject

diff --git a/Engine/Plugins/Animation/LiveLink/Source/LiveLinkComponents/Public/Controllers/LiveLinkTransformController.h b/Engine/Plugins/Animation/LiveLink/Source/LiveLinkComponents/Public/Controllers/LiveLinkTransformController.h
--- a/Engine/Plugins/Animation/LiveLink/Source/LiveLinkComponents/Public/Controllers/LiveLinkTransformController.h
+++ b/Engine/Plugins/Animation/LiveLink/Source/LiveLinkComponents/Public/Controllers/LiveLinkTransformController.h
@@ -10,7 +10,7 @@ struct FTransform;
 class USceneComponent;
 
 USTRUCT(BlueprintType)
-struct FLiveLinkTransformControllerData
+struct LIVELINKCOMPONENTS_API FLiveLinkTransformControllerData
 {
 	GENERATED_BODY()
 
diff --git a/Engine/Source/Runtime/LiveLinkInterface/Private/Roles/LiveLinkBasicRole.cpp b/Engine/Source/Runtime/LiveLinkInterface/Private/Roles/LiveLinkBasicRole.cpp
--- a/Engine/Source/Runtime/LiveLinkInterface/Private/Roles/LiveLinkBasicRole.cpp
+++ b/Engine/Source/Runtime/LiveLinkInterface/Private/Roles/LiveLinkBasicRole.cpp
@@ -278,13 +278,17 @@ void ULiveLinkBasicFrameInterpolateProcessor::FLiveLinkBasicFrameInterpolateProc
 		TArray<float>& PropertiesResult = OutBlendedFrameData.GetBaseData()->PropertyValues;
 
 		int32 NumOfProperties = FMath::Min(PropertiesA.Num(), PropertiesB.Num());
-		check(PropertiesResult.Num() <= NumOfProperties);
+		PropertiesResult.SetNum(NumOfProperties);
 
 		for (int32 PropertyIndex = 0; PropertyIndex < NumOfProperties; ++PropertyIndex)
 		{
 			PropertiesResult[PropertyIndex] = FMath::Lerp(PropertiesA[PropertyIndex], PropertiesB[PropertyIndex], InBlendWeight);
 		}
 	}
+	else
+	{
+		OutBlendedFrameData.GetBaseData()->PropertyValues = FrameWhenCanNotBlend.GetBaseData()->PropertyValues;
+	}
 
 	if (Options.bInterpolateInterpProperties)
 	{
