Fix deprecation warning from ILiveLInkSource for FacialAnimLiveLinkSource

diff --git a/Engine/Source/Runtime/LiveLinkInterface/Public/ILiveLinkSource.h b/Engine/Source/Runtime/LiveLinkInterface/Public/ILiveLinkSource.h
--- a/Engine/Source/Runtime/LiveLinkInterface/Public/ILiveLinkSource.h
+++ b/Engine/Source/Runtime/LiveLinkInterface/Public/ILiveLinkSource.h
@@ -24,6 +24,8 @@ public:
 
 	UE_DEPRECATED(4.23, "ILiveLinkSource::IsSourceStillValid is deprecated. Please use IsSourceValid instead!")
 	virtual bool IsSourceStillValid() { return IsSourceValid(); }
+	UE_DEPRECATED(4.23, "ILiveLinkSource::IsSourceStillValid is deprecated. Please use IsSourceValid instead!")
+	virtual bool IsSourceStillValid() const { return IsSourceValid(); }
 
 	// Can this source be displayed in the Source UI list
 	virtual bool CanBeDisplayedInUI() const { return true; }
