Merging EngineMerge in bringing 4.23 engine.
This represents Main at CL# 7706454  and Release-4.23 at CL#7706002
#rnx

diff --git a/Engine/Source/Runtime/LiveLinkInterface/Public/Roles/LiveLinkAnimationBlueprintStructs.h b/Engine/Source/Runtime/LiveLinkInterface/Public/Roles/LiveLinkAnimationBlueprintStructs.h
--- a/Engine/Source/Runtime/LiveLinkInterface/Public/Roles/LiveLinkAnimationBlueprintStructs.h
+++ b/Engine/Source/Runtime/LiveLinkInterface/Public/Roles/LiveLinkAnimationBlueprintStructs.h
@@ -2,146 +2,14 @@
 
 #pragma once
 
-#include "CoreMinimal.h"
-#include "ILiveLinkSource.h"
-#include "LiveLinkTypes.h"
-#include "Misc/Timecode.h"
-#include "Misc/FrameRate.h"
-#include "Roles/LiveLinkAnimationTypes.h"
-#include "LiveLinkAnimationBlueprintStructs.generated.h"
+#include "Roles/LiveLinkAnimationBlueprintStructs.h"
 
-USTRUCT(BlueprintType)
-struct FSubjectMetadata
+UE_DEPRECATED(4.23, "Including LiveLinkBlueprintStructs.h has been deprecated. Please include Roles/LiveLinkAnimationBlueprintStructs.h instead.")
+inline void LiveLinkBlueprintStructsHeaderDeprecatedWarning()
 {
-	GENERATED_USTRUCT_BODY()
+}
 
-	UPROPERTY(BlueprintReadWrite, EditAnywhere, Category = "LiveLink")
-	TMap<FName, FString> StringMetadata;
-
-	UPROPERTY(BlueprintReadWrite, EditAnywhere, Category = "LiveLink")
-	FTimecode SceneTimecode;
-
-	UPROPERTY(BlueprintReadWrite, EditAnywhere, Category = "LiveLink")
-	FFrameRate SceneFramerate;
-};
-
-USTRUCT()
-struct FCachedSubjectFrame
-{
-	GENERATED_USTRUCT_BODY()
-
-	FCachedSubjectFrame();
-
-	FCachedSubjectFrame(const FLiveLinkSkeletonStaticData* InStaticData, const FLiveLinkAnimationFrameData* InAnimData);
-
-	virtual ~FCachedSubjectFrame() = default;
-
-	void SetCurvesFromCache(TMap<FName, float>& OutCurves);
-
-	void GetSubjectMetadata(FSubjectMetadata& OutSubjectMetadata);
-
-	int32 GetNumberOfTransforms();
-
-	void GetTransformNames(TArray<FName>& OutTransformNames);
-
-	void GetTransformName(const int32 InTransformIndex, FName& OutName);
-
-	int32 GetTransformIndexFromName(FName InTransformName);
-
-	int32 GetParentTransformIndex(const int32 InTransformIndex);
-
-	void GetChildTransformIndices(const int32 InTransformIndex, TArray<int32>& OutChildIndices);
-
-	void GetTransformParentSpace(const int32 InTransformIndex, FTransform& OutTransform);
-
-	void GetTransformRootSpace(const int32 InTransformIndex, FTransform& OutTransform);
-
-	int32 GetRootIndex();
-
-	FLiveLinkSkeletonStaticData& GetSourceSkeletonData() { return SourceSkeletonData; }
-	const FLiveLinkSkeletonStaticData& GetSourceSkeletonData() const { return SourceSkeletonData; }
-
-	FLiveLinkAnimationFrameData& GetSourceAnimationFrameData() { return SourceAnimationFrameData; }
-	const FLiveLinkAnimationFrameData& GetSourceAnimationFrameData() const { return SourceAnimationFrameData; }
-
-private:
-	FLiveLinkSkeletonStaticData SourceSkeletonData;
-	FLiveLinkAnimationFrameData SourceAnimationFrameData;
-	TArray<TPair<bool, FTransform>> RootSpaceTransforms;
-	TArray<TPair<bool, TArray<int32>>> ChildTransformIndices;
-	TMap<FName, float> CachedCurves;
-
-	bool bHaveCachedCurves;
-
-	void CacheCurves();
-
-	bool IsValidTransformIndex(int32 InTransformIndex);
-};
-
-USTRUCT(BlueprintType)
-struct LIVELINKINTERFACE_API FLiveLinkTransform
-{
-	GENERATED_USTRUCT_BODY()
-
-	FLiveLinkTransform();
-
-	virtual ~FLiveLinkTransform() = default;
-
-	void GetName(FName& Name);
-
-	void GetTransformParentSpace(FTransform& OutTransform);
-
-	void GetTransformRootSpace(FTransform& OutTransform);
-
-	bool HasParent();
-
-	void GetParent(FLiveLinkTransform& OutParentTransform);
-
-	int32 GetChildCount();
-
-	void GetChildren(TArray<FLiveLinkTransform>& OutChildTransforms);
-
-	void SetCachedFrame(TSharedPtr<FCachedSubjectFrame> InCachedFrame);
-
-	void SetTransformIndex(const int32 InTransformIndex);
-
-	int32 GetTransformIndex() const;
-
-private:
-	TSharedPtr<FCachedSubjectFrame> CachedFrame;
-	int32 TransformIndex;
-};
-
-USTRUCT(BlueprintType)
-struct LIVELINKINTERFACE_API FSubjectFrameHandle : public FLiveLinkBaseBlueprintData
+inline void TriggerLiveLinkBlueprintStructsHeaderDeprecatedWarning()
 {
-	GENERATED_USTRUCT_BODY()
-
-	FSubjectFrameHandle() = default;
-
-	virtual ~FSubjectFrameHandle() = default;
-
-	void GetCurves(TMap<FName, float>& OutCurves);
-
-	void GetSubjectMetadata(FSubjectMetadata& OutMetadata);
-
-	int32 GetNumberOfTransforms();
-
-	void GetTransformNames(TArray<FName>& OutTransformNames);
-
-	void GetRootTransform(FLiveLinkTransform& OutLiveLinkTransform);
-
-	void GetTransformByIndex(int32 InTransformIndex, FLiveLinkTransform& OutLiveLinkTransform);
-
-	void GetTransformByName(FName InTransformName, FLiveLinkTransform& OutLiveLinkTransform);
-
-	void SetCachedFrame(TSharedPtr<FCachedSubjectFrame> InCachedFrame);
-
-	FLiveLinkSkeletonStaticData* GetSourceSkeletonStaticData();
-
-	FLiveLinkAnimationFrameData* GetSourceAnimationFrameData();
-
-private:
-	TSharedPtr<FCachedSubjectFrame> CachedFrame;
-};
-
+	LiveLinkBlueprintStructsHeaderDeprecatedWarning();
+}
