- Fix for livelink virtual subjects api potentially crashing when used multi threaded

diff --git a/Engine/Plugins/Animation/LiveLink/Source/LiveLink/Private/LiveLinkAnimationVirtualSubject.cpp b/Engine/Plugins/Animation/LiveLink/Source/LiveLink/Private/LiveLinkAnimationVirtualSubject.cpp
--- a/Engine/Plugins/Animation/LiveLink/Source/LiveLink/Private/LiveLinkAnimationVirtualSubject.cpp
+++ b/Engine/Plugins/Animation/LiveLink/Source/LiveLink/Private/LiveLinkAnimationVirtualSubject.cpp
@@ -124,8 +124,8 @@ void ULiveLinkAnimationVirtualSubject::BuildSkeleton(const TArray<FLiveLinkSubje
 {
 	if (DoesSkeletonNeedRebuilding())
 	{
-		FrameSnapshot.StaticData.InitializeWith(FLiveLinkSkeletonStaticData::StaticStruct(), nullptr);
-		FLiveLinkSkeletonStaticData* SkeletonData = FrameSnapshot.StaticData.Cast<FLiveLinkSkeletonStaticData>();
+		FLiveLinkStaticDataStruct StaticData(FLiveLinkSkeletonStaticData::StaticStruct());
+		FLiveLinkSkeletonStaticData* SkeletonData = StaticData.Cast<FLiveLinkSkeletonStaticData>();
 
 		TArray<FName> BoneNames{ TEXT("Root") };
 		TArray<int32> BoneParents{ INDEX_NONE };
@@ -146,23 +146,21 @@ void ULiveLinkAnimationVirtualSubject::BuildSkeleton(const TArray<FLiveLinkSubje
 		SkeletonData->SetBoneNames(BoneNames);
 		SkeletonData->SetBoneParents(BoneParents);
 
+		UpdateStaticDataSnapshot(MoveTemp(StaticData));
+
 		bInvalidate = false;
 	}
 }
 
 void ULiveLinkAnimationVirtualSubject::BuildFrame(const TArray<FLiveLinkSubjectFrameData>& InSubjectSnapshots)
 {
-	if (!FrameSnapshot.FrameData.IsValid())
-	{
-		FrameSnapshot.FrameData.InitializeWith(FLiveLinkAnimationFrameData::StaticStruct(), nullptr);
-	}
+	const FLiveLinkSkeletonStaticData* SnapshotSkeletonData = GetFrameSnapshot().StaticData.Cast<FLiveLinkSkeletonStaticData>();
+	FLiveLinkFrameDataStruct NewFrameData(FLiveLinkAnimationFrameData::StaticStruct());
+	FLiveLinkAnimationFrameData* NewSnapshotFrameData = NewFrameData.Cast<FLiveLinkAnimationFrameData>();
 
-	FLiveLinkSkeletonStaticData* SnapshotSkeletonData = FrameSnapshot.StaticData.Cast<FLiveLinkSkeletonStaticData>();
-	FLiveLinkAnimationFrameData* SnapshotFrameData = FrameSnapshot.FrameData.Cast<FLiveLinkAnimationFrameData>();
-
-	SnapshotFrameData->Transforms.Reset(SnapshotSkeletonData->GetBoneNames().Num());
-	SnapshotFrameData->Transforms.Add(FTransform::Identity);
-	SnapshotFrameData->MetaData.StringMetaData.Empty();
+	NewSnapshotFrameData->Transforms.Reset(SnapshotSkeletonData->GetBoneNames().Num());
+	NewSnapshotFrameData->Transforms.Add(FTransform::Identity);
+	NewSnapshotFrameData->MetaData.StringMetaData.Empty();
 
 	//Go over each subject snapshot and take transforms and curves
 	check(InSubjectSnapshots.Num() == Subjects.Num());
@@ -172,19 +170,21 @@ void ULiveLinkAnimationVirtualSubject::BuildFrame(const TArray<FLiveLinkSubjectF
 		check(SubjectSnapShotData.FrameData.IsValid());
 		const FLiveLinkAnimationFrameData* SubjectFrameData = SubjectSnapShotData.FrameData.Cast<FLiveLinkAnimationFrameData>();
 
-		SnapshotFrameData->Transforms.Append(SubjectFrameData->Transforms);
-		SnapshotFrameData->PropertyValues.Append(SubjectFrameData->PropertyValues);
+		NewSnapshotFrameData->Transforms.Append(SubjectFrameData->Transforms);
+		NewSnapshotFrameData->PropertyValues.Append(SubjectFrameData->PropertyValues);
 		for (const auto& MetaDatum : SubjectFrameData->MetaData.StringMetaData)
 		{
 			const FName QualifiedKey = FName(*(Subjects[i].ToString() + MetaDatum.Key.ToString()));
-			SnapshotFrameData->MetaData.StringMetaData.Emplace(Subjects[i], MetaDatum.Value);
+			NewSnapshotFrameData->MetaData.StringMetaData.Emplace(Subjects[i], MetaDatum.Value);
 		}
 	}
+
+	UpdateFrameDataSnapshot(MoveTemp(NewFrameData));
 }
 
 bool ULiveLinkAnimationVirtualSubject::DoesSkeletonNeedRebuilding() const
 {
-	return !FrameSnapshot.StaticData.IsValid() || bInvalidate;
+	return !HasValidStaticData() || bInvalidate;
 }
 
 #if WITH_EDITOR
diff --git a/Engine/Plugins/Animation/LiveLink/Source/LiveLink/Private/VirtualSubjects/LiveLinkBlueprintVirtualSubject.cpp b/Engine/Plugins/Animation/LiveLink/Source/LiveLink/Private/VirtualSubjects/LiveLinkBlueprintVirtualSubject.cpp
--- a/Engine/Plugins/Animation/LiveLink/Source/LiveLink/Private/VirtualSubjects/LiveLinkBlueprintVirtualSubject.cpp
+++ b/Engine/Plugins/Animation/LiveLink/Source/LiveLink/Private/VirtualSubjects/LiveLinkBlueprintVirtualSubject.cpp
@@ -9,8 +9,14 @@ void ULiveLinkBlueprintVirtualSubject::Initialize(FLiveLinkSubjectKey InSubjectK
 	Super::Initialize(InSubjectKey, Role, LiveLinkClient);
 
 	CachedStaticData.InitializeWith(GetRoleStaticStruct(), nullptr);
-	FrameSnapshot.StaticData.InitializeWith(GetRoleStaticStruct(), nullptr);
-	FrameSnapshot.FrameData.InitializeWith(GetRoleFrameStruct(), nullptr);
+
+	FLiveLinkStaticDataStruct NewStaticData;
+	NewStaticData.InitializeWith(GetRoleStaticStruct(), CachedStaticData.GetBaseData());
+	UpdateStaticDataSnapshot(MoveTemp(NewStaticData));
+
+	FLiveLinkFrameDataStruct NewFrameData;
+	NewFrameData.InitializeWith(GetRoleFrameStruct(), nullptr);
+	UpdateFrameDataSnapshot(MoveTemp(NewFrameData));
 
 	{
 		FEditorScriptExecutionGuard ScriptGuard;
@@ -29,29 +35,35 @@ void ULiveLinkBlueprintVirtualSubject::Update()
 		OnUpdate();
 	}
 
-	if (FrameSnapshot.FrameData.IsValid() && !FrameSnapshot.StaticData.IsValid() && CachedStaticData.IsValid())
+	if (HasValidFrameData() && !HasValidStaticData() && CachedStaticData.IsValid())
 	{
-		FrameSnapshot.StaticData.InitializeWith(CachedStaticData);
+		FLiveLinkStaticDataStruct NewStaticData;
+		NewStaticData.InitializeWith(GetRoleStaticStruct(), CachedStaticData.GetBaseData());
+		UpdateStaticDataSnapshot(MoveTemp(NewStaticData));
 	}
 }
 
 void ULiveLinkBlueprintVirtualSubject::UpdateVirtualSubjectStaticData(const FLiveLinkBaseStaticData* InStaticData)
 {
 	CachedStaticData.InitializeWith(GetRoleStaticStruct(), InStaticData);
-	FrameSnapshot.StaticData.InitializeWith(GetRoleStaticStruct(), InStaticData);
+
+	FLiveLinkStaticDataStruct NewStaticData;
+	NewStaticData.InitializeWith(GetRoleStaticStruct(), InStaticData);
+	UpdateStaticDataSnapshot(MoveTemp(NewStaticData));
 	
 	// Invalidate any existing Frame Data
-	FrameSnapshot.FrameData.Reset();
+	InvalidateFrameData();
 }
 
 void ULiveLinkBlueprintVirtualSubject::UpdateVirtualSubjectFrameData(const FLiveLinkBaseFrameData* InFrameData, bool bInShouldStampCurrentTime)
 {
-	FrameSnapshot.FrameData.InitializeWith(GetRoleFrameStruct(), InFrameData);
+	FLiveLinkFrameDataStruct NewFrameData;
+	NewFrameData.InitializeWith(GetRoleFrameStruct(), InFrameData);
 
 	// Stamp the current time into the frame if desired
 	if (bInShouldStampCurrentTime)
 	{
-		if (FLiveLinkBaseFrameData* BaseFrameData = FrameSnapshot.FrameData.GetBaseData())
+		if (FLiveLinkBaseFrameData* BaseFrameData = NewFrameData.GetBaseData())
 		{
 			// Stamp the current world time
 			BaseFrameData->WorldTime = FLiveLinkWorldTime(FApp::GetCurrentTime());
@@ -64,6 +76,8 @@ void ULiveLinkBlueprintVirtualSubject::UpdateVirtualSubjectFrameData(const FLive
 			}
 		}
 	}
+
+	UpdateFrameDataSnapshot(MoveTemp(NewFrameData));
 }
 
 bool ULiveLinkBlueprintVirtualSubject::UpdateVirtualSubjectStaticData_Internal(const FLiveLinkBaseStaticData& InStruct)
diff --git a/Engine/Plugins/Runtime/LiveLinkOvernDisplay/Source/LiveLinkOverNDisplay/Private/NDisplaySlaveVirtualSubject.cpp b/Engine/Plugins/Runtime/LiveLinkOvernDisplay/Source/LiveLinkOverNDisplay/Private/NDisplaySlaveVirtualSubject.cpp
--- a/Engine/Plugins/Runtime/LiveLinkOvernDisplay/Source/LiveLinkOverNDisplay/Private/NDisplaySlaveVirtualSubject.cpp
+++ b/Engine/Plugins/Runtime/LiveLinkOvernDisplay/Source/LiveLinkOverNDisplay/Private/NDisplaySlaveVirtualSubject.cpp
@@ -14,7 +14,7 @@ void UNDisplaySlaveVirtualSubject::Update()
 
 void UNDisplaySlaveVirtualSubject::UpdateFrameData(FLiveLinkFrameDataStruct&& NewFrameData)
 {
-	FrameSnapshot.FrameData = MoveTemp(NewFrameData);
+	UpdateFrameDataSnapshot(MoveTemp(NewFrameData));
 }
 
 void UNDisplaySlaveVirtualSubject::SetTrackedSubjectInfo(const FLiveLinkSubjectKey& InSubjectKey, TSubclassOf<ULiveLinkRole> InRole)
diff --git a/Engine/Source/Runtime/LiveLinkInterface/Private/ILiveLinkSubject.cpp b/Engine/Source/Runtime/LiveLinkInterface/Private/ILiveLinkSubject.cpp
--- a/Engine/Source/Runtime/LiveLinkInterface/Private/ILiveLinkSubject.cpp
+++ b/Engine/Source/Runtime/LiveLinkInterface/Private/ILiveLinkSubject.cpp
@@ -70,13 +70,6 @@ bool ILiveLinkSubject::SupportsRole(TSubclassOf<ULiveLinkRole> InDesiredRole) co
 }
 
 
-bool ILiveLinkSubject::HasValidFrameSnapshot() const
-{
-	const FLiveLinkSubjectFrameData& FrameSnapshot = GetFrameSnapshot();
-	return FrameSnapshot.StaticData.IsValid() && FrameSnapshot.FrameData.IsValid();
-}
-
-
 bool ILiveLinkSubject::Translate(const ILiveLinkSubject* InLinkSubject, TSubclassOf<ULiveLinkRole> InDesiredRole, const FLiveLinkStaticDataStruct& InStaticData, const FLiveLinkFrameDataStruct& InFrameData, FLiveLinkSubjectFrameData& OutFrame)
 {
 	// Find one that matches exactly
diff --git a/Engine/Source/Runtime/LiveLinkInterface/Private/LiveLinkVirtualSubject.cpp b/Engine/Source/Runtime/LiveLinkInterface/Private/LiveLinkVirtualSubject.cpp
--- a/Engine/Source/Runtime/LiveLinkInterface/Private/LiveLinkVirtualSubject.cpp
+++ b/Engine/Source/Runtime/LiveLinkInterface/Private/LiveLinkVirtualSubject.cpp
@@ -19,22 +19,30 @@ void ULiveLinkVirtualSubject::Initialize(FLiveLinkSubjectKey InSubjectKey, TSubc
 void ULiveLinkVirtualSubject::Update()
 {
 	// Invalid the snapshot
-	FrameSnapshot.StaticData.Reset();
-	FrameSnapshot.FrameData.Reset();
+	InvalidateStaticData();
+	InvalidateFrameData();
 
 	UpdateTranslatorsForThisFrame();
 }
 
 
+bool ULiveLinkVirtualSubject::EvaluateFrame(TSubclassOf<ULiveLinkRole> InDesiredRole, FLiveLinkSubjectFrameData& OutFrame)
+{
+	//Protect our data when being evaluated
+	FScopeLock Lock(&SnapshotAccessCriticalSection);
+	return ILiveLinkSubject::EvaluateFrame(InDesiredRole, OutFrame);
+}
+
 void ULiveLinkVirtualSubject::ClearFrames()
 {
-	FrameSnapshot.StaticData.Reset();
+	FScopeLock Lock(&SnapshotAccessCriticalSection);
+	CurrentFrameSnapshot.StaticData.Reset();
 }
 
 
 bool ULiveLinkVirtualSubject::HasValidFrameSnapshot() const
 {
-	return FrameSnapshot.StaticData.IsValid() && FrameSnapshot.FrameData.IsValid();
+	return CurrentFrameSnapshot.StaticData.IsValid() && CurrentFrameSnapshot.FrameData.IsValid();
 }
 
 TArray<FLiveLinkTime> ULiveLinkVirtualSubject::GetFrameTimes() const
@@ -45,10 +53,20 @@ TArray<FLiveLinkTime> ULiveLinkVirtualSubject::GetFrameTimes() const
 	}
 
 	TArray<FLiveLinkTime> Result;
-	Result.Emplace(FrameSnapshot.FrameData.GetBaseData()->WorldTime.GetOffsettedTime(), FrameSnapshot.FrameData.GetBaseData()->MetaData.SceneTime);
+	Result.Emplace(CurrentFrameSnapshot.FrameData.GetBaseData()->WorldTime.GetOffsettedTime(), CurrentFrameSnapshot.FrameData.GetBaseData()->MetaData.SceneTime);
 	return Result;
 }
 
+bool ULiveLinkVirtualSubject::HasValidStaticData() const
+{
+	return CurrentFrameSnapshot.StaticData.IsValid();
+}
+
+bool ULiveLinkVirtualSubject::HasValidFrameData() const
+{
+	return CurrentFrameSnapshot.FrameData.IsValid();
+}
+
 bool ULiveLinkVirtualSubject::DependsOnSubject(FName SubjectName) const
 {
 	return Subjects.Contains(SubjectName);
@@ -70,3 +88,28 @@ void ULiveLinkVirtualSubject::UpdateTranslatorsForThisFrame()
 		}
 	}
 }
+
+void ULiveLinkVirtualSubject::UpdateStaticDataSnapshot(FLiveLinkStaticDataStruct&& NewStaticData)
+{
+	FScopeLock Lock(&SnapshotAccessCriticalSection);
+	CurrentFrameSnapshot.StaticData = MoveTemp(NewStaticData);
+}
+
+void ULiveLinkVirtualSubject::UpdateFrameDataSnapshot(FLiveLinkFrameDataStruct&& NewFrameData)
+{
+	FScopeLock Lock(&SnapshotAccessCriticalSection);
+	CurrentFrameSnapshot.FrameData = MoveTemp(NewFrameData);
+}
+
+void ULiveLinkVirtualSubject::InvalidateStaticData()
+{
+	FScopeLock Lock(&SnapshotAccessCriticalSection);
+	CurrentFrameSnapshot.StaticData.Reset();
+}
+
+void ULiveLinkVirtualSubject::InvalidateFrameData()
+{
+	FScopeLock Lock(&SnapshotAccessCriticalSection);
+	CurrentFrameSnapshot.FrameData.Reset();
+}
+
diff --git a/Engine/Source/Runtime/LiveLinkInterface/Public/LiveLinkVirtualSubject.h b/Engine/Source/Runtime/LiveLinkInterface/Public/LiveLinkVirtualSubject.h
--- a/Engine/Source/Runtime/LiveLinkInterface/Public/LiveLinkVirtualSubject.h
+++ b/Engine/Source/Runtime/LiveLinkInterface/Public/LiveLinkVirtualSubject.h
@@ -25,21 +25,28 @@ class LIVELINKINTERFACE_API ULiveLinkVirtualSubject : public UObject, public ILi
 public:
 	virtual void Initialize(FLiveLinkSubjectKey SubjectKey, TSubclassOf<ULiveLinkRole> Role, ILiveLinkClient* LiveLinkClient) override;
 	virtual void Update() override;
+	virtual bool EvaluateFrame(TSubclassOf<ULiveLinkRole> InDesiredRole, FLiveLinkSubjectFrameData& OutFrame) override;
 	virtual void ClearFrames() override;
 	virtual FLiveLinkSubjectKey GetSubjectKey() const override { return SubjectKey; }
 	virtual TSubclassOf<ULiveLinkRole> GetRole() const override { return Role; }
 	virtual bool HasValidFrameSnapshot() const override;
-	virtual FLiveLinkStaticDataStruct& GetStaticData() override { return FrameSnapshot.StaticData; }
-	virtual const FLiveLinkStaticDataStruct& GetStaticData() const override { return FrameSnapshot.StaticData; }
+	virtual FLiveLinkStaticDataStruct& GetStaticData() override { return CurrentFrameSnapshot.StaticData; }
+	virtual const FLiveLinkStaticDataStruct& GetStaticData() const override { return CurrentFrameSnapshot.StaticData; }
 	virtual const TArray<ULiveLinkFrameTranslator::FWorkerSharedPtr> GetFrameTranslators() const override { return CurrentFrameTranslators; }
 	virtual TArray<FLiveLinkTime> GetFrameTimes() const override;
 	virtual bool IsRebroadcasted() const override { return bRebroadcastSubject; }
 	virtual bool HasStaticDataBeenRebroadcasted() const override { return bHasStaticDataBeenRebroadcast; }
 	virtual void SetStaticDataAsRebroadcasted(const bool bInSent) override { bHasStaticDataBeenRebroadcast = bInSent; }
 protected:
-	virtual const FLiveLinkSubjectFrameData& GetFrameSnapshot() const override { return FrameSnapshot; }
+	virtual const FLiveLinkSubjectFrameData& GetFrameSnapshot() const override { return CurrentFrameSnapshot; }
 	//~ End ILiveLinkSubject Interface
 
+	/** Whether snapshot has valid static data */
+	bool HasValidStaticData() const;
+
+	/** Whether snapshot has valid frame data */
+	bool HasValidFrameData() const;
+
 public:
 	ILiveLinkClient* GetClient() const { return LiveLinkClient; }
 
@@ -50,14 +57,28 @@ public:
 	const TArray<ULiveLinkFrameTranslator*>& GetTranslators() const { return FrameTranslators; }
 
 	/** Returns the current frame data of this virtual subject */
-	const FLiveLinkFrameDataStruct& GetFrameData() const { return FrameSnapshot.FrameData; }
+	const FLiveLinkFrameDataStruct& GetFrameData() const { return CurrentFrameSnapshot.FrameData; }
 
 	/** Returns true whether this virtual subject depends on the Subject named SubjectName */
 	virtual bool DependsOnSubject(FName SubjectName) const;
 	
 protected:
+
+	/** Updates the list of translators valid for this frame */
 	void UpdateTranslatorsForThisFrame();
 
+	/** Updates our snapshot's static data */
+	void UpdateStaticDataSnapshot(FLiveLinkStaticDataStruct&& NewStaticData);
+
+	/** Updates our snapshot's frame data */
+	void UpdateFrameDataSnapshot(FLiveLinkFrameDataStruct&& NewFrameData);
+
+	/** Invalidates our snapshot's static data */
+	void InvalidateStaticData();
+
+	/** Invalidates our snapshot's frame data */
+	void InvalidateFrameData();
+
 protected:
 	/** The role the subject was build with. */
 	UPROPERTY()
@@ -78,7 +99,7 @@ protected:
 	/** LiveLinkClient to get access to subjects */
 	ILiveLinkClient* LiveLinkClient;
 
-	/** Last evaluated frame for this subject. */
+	UE_DEPRECATED(4.27, "VirtualSubject FrameSnapshot is now private to have thread safe accesses. Please use UpdateStaticDataSnapshot or UpdateFrameDataSnapshot to update its value")
 	FLiveLinkSubjectFrameData FrameSnapshot;
 
 	/** Name of the subject */
@@ -87,6 +108,17 @@ protected:
 	/** If true, static data has been sent for this rebroadcast */
 	bool bHasStaticDataBeenRebroadcast = false;
 
+	/** Lock to protect the FrameSnapshot. 
+	 * VirtualSubjects can be manipulated from anywhere versus LiveSubjects 
+	 * that have an access controlled in the Source Collection.
+	 * Evaluating subjects is AnyThread so we can be evaluated 
+	 * while our snapshot is getting set
+	 */
+	mutable FCriticalSection SnapshotAccessCriticalSection;
+
 private:
 	TArray<ULiveLinkFrameTranslator::FWorkerSharedPtr> CurrentFrameTranslators;
+
+	/** Last evaluated frame for this subject. */
+	FLiveLinkSubjectFrameData CurrentFrameSnapshot;
 };
