Fixed mismatched copy constructor/assignment operator in FLiveLinkSubjectKey.

diff --git a/Engine/Source/Runtime/LiveLinkInterface/Public/LiveLinkTypes.h b/Engine/Source/Runtime/LiveLinkInterface/Public/LiveLinkTypes.h
--- a/Engine/Source/Runtime/LiveLinkInterface/Public/LiveLinkTypes.h
+++ b/Engine/Source/Runtime/LiveLinkInterface/Public/LiveLinkTypes.h
@@ -88,7 +88,6 @@ struct LIVELINKINTERFACE_API FLiveLinkSubjectKey
 
 	FLiveLinkSubjectKey() = default;
 	FLiveLinkSubjectKey(FGuid InSource, FName InSubjectName) : Source(InSource), SubjectName(InSubjectName) {}
-	FLiveLinkSubjectKey(const FLiveLinkSubjectKey& Rhs) : Source(Rhs.Source), SubjectName(Rhs.SubjectName) {}
 
 	bool operator== (const FLiveLinkSubjectKey& Other) const { return SubjectName == Other.SubjectName && Source == Other.Source; }
 	bool operator!=(const FLiveLinkSubjectKey& Other) const	{ return !(*this == Other); }
