Add export macro to FLiveLinkSubjectName

#rnx

diff --git a/Engine/Source/Runtime/LiveLinkInterface/Public/LiveLinkTypes.h b/Engine/Source/Runtime/LiveLinkInterface/Public/LiveLinkTypes.h
--- a/Engine/Source/Runtime/LiveLinkInterface/Public/LiveLinkTypes.h
+++ b/Engine/Source/Runtime/LiveLinkInterface/Public/LiveLinkTypes.h
@@ -36,7 +36,7 @@
 
 
 USTRUCT(BlueprintType)
-struct FLiveLinkSubjectName
+struct LIVELINKINTERFACE_API FLiveLinkSubjectName
 {
 public:
 	GENERATED_BODY()
