LiveLink: Added BP function to query subject state (invalid, disconnected, unresponsive).

This should mimic the state indicated in the live link panel for a given subject. The function is blueprint callable, and is called GetLiveLinkSubjectState. The value returned is an enum, as follows:

enum class ELiveLinkSubjectState : uint8
{
  Connected,
  Unresponsive,
  Disconnected,
  InvalidOrDisabled,
  Unknown
};

It is similar to ETimedDataInputState but can report more states suitable to a higher level query.

diff --git a/Engine/Plugins/Animation/LiveLink/Source/LiveLink/Private/LiveLinkBlueprintLibrary.cpp b/Engine/Plugins/Animation/LiveLink/Source/LiveLink/Private/LiveLinkBlueprintLibrary.cpp
--- a/Engine/Plugins/Animation/LiveLink/Source/LiveLink/Private/LiveLinkBlueprintLibrary.cpp
+++ b/Engine/Plugins/Animation/LiveLink/Source/LiveLink/Private/LiveLinkBlueprintLibrary.cpp
@@ -242,6 +242,21 @@ bool ULiveLinkBlueprintLibrary::IsLiveLinkSubjectEnabled(const FLiveLinkSubjectN
 	}
 	return false;
 }
+
+ELiveLinkSubjectState ULiveLinkBlueprintLibrary::GetLiveLinkSubjectState(const FLiveLinkSubjectName SubjectName)
+{
+	IModularFeatures& ModularFeatures = IModularFeatures::Get();
+
+	if (!ModularFeatures.IsModularFeatureAvailable(ILiveLinkClient::ModularFeatureName))
+	{
+		return ELiveLinkSubjectState::Unknown;
+	}
+
+	ILiveLinkClient& LiveLinkClient = ModularFeatures.GetModularFeature<ILiveLinkClient>(ILiveLinkClient::ModularFeatureName);
+	
+	return LiveLinkClient.GetSubjectState(SubjectName);
+}
+
 void ULiveLinkBlueprintLibrary::SetLiveLinkSubjectEnabled(const FLiveLinkSubjectKey SubjectKey, bool bEnabled)
 {
 	IModularFeatures& ModularFeatures = IModularFeatures::Get();
diff --git a/Engine/Plugins/Animation/LiveLink/Source/LiveLink/Private/LiveLinkClient.cpp b/Engine/Plugins/Animation/LiveLink/Source/LiveLink/Private/LiveLinkClient.cpp
--- a/Engine/Plugins/Animation/LiveLink/Source/LiveLink/Private/LiveLinkClient.cpp
+++ b/Engine/Plugins/Animation/LiveLink/Source/LiveLink/Private/LiveLinkClient.cpp
@@ -294,6 +294,42 @@ void FLiveLinkClient::RemoveRebroadcastedSubject(FLiveLinkSubjectKey InSubjectKe
 	}
 }
 
+ELiveLinkSubjectState FLiveLinkClient::GetSubjectState(FLiveLinkSubjectName InSubjectName) const
+{
+	const FLiveLinkSubjectKey* SubjectKey = EnabledSubjects.Find(InSubjectName);
+
+	if (!SubjectKey)
+	{
+		return ELiveLinkSubjectState::InvalidOrDisabled;
+	}
+
+	const FLiveLinkCollectionSubjectItem* SubjectItem = Collection->FindSubject(*SubjectKey);
+
+	if (!SubjectItem)
+	{
+		return ELiveLinkSubjectState::InvalidOrDisabled;
+	}
+
+	const FLiveLinkSubject* LiveSubject = SubjectItem->GetLiveSubject();
+
+	if (!LiveSubject)
+	{
+		return ELiveLinkSubjectState::InvalidOrDisabled;
+	}
+
+	switch (const ETimedDataInputState InputState = LiveSubject->GetState())
+	{
+	case ETimedDataInputState::Connected: return ELiveLinkSubjectState::Connected;
+	case ETimedDataInputState::Unresponsive: return ELiveLinkSubjectState::Unresponsive;
+	case ETimedDataInputState::Disconnected: return ELiveLinkSubjectState::Disconnected;
+	default:
+		ensureMsgf(false, TEXT("Unhandled ETimedDataInputState::%d"), InputState);
+		return ELiveLinkSubjectState::Unknown;
+	}
+
+	checkNoEntry();
+}
+
 FGuid FLiveLinkClient::AddSource(TSharedPtr<ILiveLinkSource> InSource)
 {
 	check(Collection);
diff --git a/Engine/Plugins/Animation/LiveLink/Source/LiveLink/Public/LiveLinkBlueprintLibrary.h b/Engine/Plugins/Animation/LiveLink/Source/LiveLink/Public/LiveLinkBlueprintLibrary.h
--- a/Engine/Plugins/Animation/LiveLink/Source/LiveLink/Public/LiveLinkBlueprintLibrary.h
+++ b/Engine/Plugins/Animation/LiveLink/Source/LiveLink/Public/LiveLinkBlueprintLibrary.h
@@ -4,6 +4,7 @@
 
 
 #include "Kismet/BlueprintFunctionLibrary.h"
+#include "ILiveLinkClient.h"
 
 #include "LiveLinkBlueprintLibrary.generated.h"
 
@@ -162,6 +163,12 @@ public:
 	UFUNCTION(BlueprintCallable, Category = "LiveLink")
 	static bool IsLiveLinkSubjectEnabled(const FLiveLinkSubjectName SubjectName);
 
+	/**
+	 * Returns the state of the subject with the given name.
+	 */
+	UFUNCTION(BlueprintCallable, Category = "LiveLink")
+	static ELiveLinkSubjectState GetLiveLinkSubjectState(const FLiveLinkSubjectName SubjectName);
+
 	/** 
 	 * Set the subject's from a specific source to enabled, disabling the other in the process.
 	 * Only 1 subject with the same name can be enabled.
diff --git a/Engine/Plugins/Animation/LiveLink/Source/LiveLink/Public/LiveLinkClient.h b/Engine/Plugins/Animation/LiveLink/Source/LiveLink/Public/LiveLinkClient.h
--- a/Engine/Plugins/Animation/LiveLink/Source/LiveLink/Public/LiveLinkClient.h
+++ b/Engine/Plugins/Animation/LiveLink/Source/LiveLink/Public/LiveLinkClient.h
@@ -82,7 +82,6 @@ public:
 	virtual FText GetSourceMachineName(FGuid EntryGuid) const override;
 	virtual bool IsSourceStillValid(FGuid EntryGuid) const override;
 
-
 	virtual void PushSubjectStaticData_AnyThread(const FLiveLinkSubjectKey& SubjectKey, TSubclassOf<ULiveLinkRole> Role, FLiveLinkStaticDataStruct&& StaticData) override;
 	virtual void PushSubjectFrameData_AnyThread(const FLiveLinkSubjectKey& SubjectKey, FLiveLinkFrameDataStruct&& FrameData) override;
 
@@ -109,6 +108,8 @@ public:
 	virtual bool IsSubjectTimeSynchronized(const FLiveLinkSubjectKey& SubjectKey) const override;
 	virtual bool IsSubjectTimeSynchronized(FLiveLinkSubjectName SubjectName) const override;
 	virtual bool IsVirtualSubject(const FLiveLinkSubjectKey& SubjectKey) const override;
+	virtual ELiveLinkSubjectState GetSubjectState(FLiveLinkSubjectName InSubjectName) const override;
+
 
 	virtual TArray<FLiveLinkSubjectKey> GetSubjectsSupportingRole(TSubclassOf<ULiveLinkRole> SupportedRole, bool bIncludeDisabledSubject, bool bIncludeVirtualSubject) const override;
 	virtual TArray<FLiveLinkTime> GetSubjectFrameTimes(const FLiveLinkSubjectKey& SubjectKey) const override;
diff --git a/Engine/Source/Runtime/LiveLinkInterface/Public/ILiveLinkClient.h b/Engine/Source/Runtime/LiveLinkInterface/Public/ILiveLinkClient.h
--- a/Engine/Source/Runtime/LiveLinkInterface/Public/ILiveLinkClient.h
+++ b/Engine/Source/Runtime/LiveLinkInterface/Public/ILiveLinkClient.h
@@ -18,6 +18,24 @@ struct FLiveLinkSubjectFrameData;
 struct FTimecode;
 class ULiveLinkSourceSettings;
 
+
+/** Describes the state of a live link subject */
+UENUM()
+enum class ELiveLinkSubjectState : uint8
+{
+	/** The input is connected. */
+	Connected,
+	/** The input is connected but no data is available. */
+	Unresponsive,
+	/** The input is not connected. */
+	Disconnected,
+	/** The subject is invalid or disabled */
+	InvalidOrDisabled,
+	/** The state of the subject is unknown. e.g. It cannot be queried */
+	Unknown
+};
+
+
 DECLARE_TS_MULTICAST_DELEGATE_OneParam(FOnLiveLinkSourceChangedDelegate, FGuid /*SourceGuid*/);
 DECLARE_TS_MULTICAST_DELEGATE_OneParam(FOnLiveLinkSubjectChangedDelegate, FLiveLinkSubjectKey /*SubjectKey*/);
 DECLARE_TS_MULTICAST_DELEGATE_OneParam(FOnLiveLinkSubjectStaticDataReceived, const FLiveLinkStaticDataStruct& /*InStaticData*/)
@@ -26,6 +44,7 @@ DECLARE_TS_MULTICAST_DELEGATE_ThreeParams(FOnLiveLinkSubjectStaticDataAdded, FLi
 DECLARE_TS_MULTICAST_DELEGATE_ThreeParams(FOnLiveLinkSubjectFrameDataAdded, FLiveLinkSubjectKey /*InSubjectKey*/, TSubclassOf<ULiveLinkRole> /*SubjectRole*/, const FLiveLinkFrameDataStruct& /*InFrameData*/)
 DECLARE_TS_MULTICAST_DELEGATE_FiveParams(FOnLiveLinkSubjectEvaluated, FLiveLinkSubjectKey /*InSubjectKey*/, TSubclassOf<ULiveLinkRole> /*RequestedRole*/, const FLiveLinkTime& /*RequestedTime*/, bool /*bResult*/, const FLiveLinkTime& /*EvaluatedFrameTime*/)
 
+
 PRAGMA_DISABLE_DEPRECATION_WARNINGS
 class ILiveLinkClient_Base_DEPRECATED : public IModularFeature
 {
@@ -253,6 +272,9 @@ public:
 	/** Whether the subject key points to a virtual subject */
 	virtual bool IsVirtualSubject(const FLiveLinkSubjectKey& SubjectKey) const = 0;
 
+	/** Returns the state of the given subject name. */
+	virtual ELiveLinkSubjectState GetSubjectState(FLiveLinkSubjectName InSubjectName) const = 0;
+
 	/** Get a list of name of subjects supporting a certain role */
 	virtual TArray<FLiveLinkSubjectKey> GetSubjectsSupportingRole(TSubclassOf<ULiveLinkRole> SupportedRole, bool bIncludeDisabledSubject, bool bIncludeVirtualSubject) const = 0;
 
@@ -367,6 +389,8 @@ public:
 	 * @return True if the subject was found and the delegates were registered. False otherwise.
 	 */
 	virtual bool RegisterForSubjectFrames(FLiveLinkSubjectName SubjectName, const FOnLiveLinkSubjectStaticDataAdded::FDelegate& OnStaticDataAdded, const FOnLiveLinkSubjectFrameDataAdded::FDelegate& OnFrameDataAddedd, FDelegateHandle& OutStaticDataAddedHandle, FDelegateHandle& OutFrameDataAddeddHandle, TSubclassOf<ULiveLinkRole>& OutSubjectRole, FLiveLinkStaticDataStruct* OutStaticData = nullptr) = 0;
+
 	/** Unregister delegates registered with RegisterForSubjectFrames. */
 	virtual void UnregisterSubjectFramesHandle(FLiveLinkSubjectName SubjectName, FDelegateHandle StaticDataAddedHandle, FDelegateHandle FrameDataAddedHandle) = 0;
+
 };
