Submit tool generated changelist from default with (1) files

diff --git a/Engine/Source/Runtime/LiveLinkInterface/Public/LiveLinkSourceSettings.h b/Engine/Source/Runtime/LiveLinkInterface/Public/LiveLinkSourceSettings.h
--- a/Engine/Source/Runtime/LiveLinkInterface/Public/LiveLinkSourceSettings.h
+++ b/Engine/Source/Runtime/LiveLinkInterface/Public/LiveLinkSourceSettings.h
@@ -117,8 +117,8 @@ struct FLiveLinkSourceBufferManagementSettings
 	UPROPERTY(EditAnywhere, Category = "Settings")
 	int32 LatestOffset = 0;
 
-	/** Maximum number of frame to keep in memory. */
-	UPROPERTY(EditAnywhere, Category = "Settings", meta=(ClampMin=1))
+	/** Maximum number of frames to keep in memory. */
+	UPROPERTY(EditAnywhere, Category = "Settings", meta=(ClampMin=1, DisplayName="Buffer Size (Frames)"))
 	int32 MaxNumberOfFrameToBuffered = 10;
 
 	/** When cleaning the buffer keep at least one frame, even if the frame doesn't matches the other options. */
