LiveLink - Add setting for the default number of buffered frames in source settings

diff --git a/Engine/Source/Runtime/LiveLinkInterface/Private/LiveLinkSourceSettings.cpp b/Engine/Source/Runtime/LiveLinkInterface/Private/LiveLinkSourceSettings.cpp
--- a/Engine/Source/Runtime/LiveLinkInterface/Private/LiveLinkSourceSettings.cpp
+++ b/Engine/Source/Runtime/LiveLinkInterface/Private/LiveLinkSourceSettings.cpp
@@ -5,6 +5,14 @@
 
 #include UE_INLINE_GENERATED_CPP_BY_NAME(LiveLinkSourceSettings)
 
+ULiveLinkSourceSettings::ULiveLinkSourceSettings()
+{
+	if (!HasAnyFlags(RF_ClassDefaultObject | RF_ArchetypeObject))
+	{
+		BufferSettings.MaxNumberOfFrameToBuffered = GetDefault<ULiveLinkDefaultSourceSettings>()->DefaultSourceFrameBufferSize;
+	}
+}
+
 void ULiveLinkSourceSettings::Serialize(FArchive& Ar)
 {
 	Super::Serialize(Ar);
@@ -37,4 +45,3 @@ bool ULiveLinkSourceSettings::CanEditChange(const FProperty* InProperty) const
 	return false;
 }
 #endif //WITH_EDITOR
-
diff --git a/Engine/Source/Runtime/LiveLinkInterface/Public/LiveLinkSourceSettings.h b/Engine/Source/Runtime/LiveLinkInterface/Public/LiveLinkSourceSettings.h
--- a/Engine/Source/Runtime/LiveLinkInterface/Public/LiveLinkSourceSettings.h
+++ b/Engine/Source/Runtime/LiveLinkInterface/Public/LiveLinkSourceSettings.h
@@ -17,6 +17,21 @@
 
 #include "LiveLinkSourceSettings.generated.h"
 
+/**
+ * Utility class that allows specifying default values for Source settings.
+ */
+UCLASS(config=Engine, defaultconfig)
+class ULiveLinkDefaultSourceSettings : public UObject
+{
+	GENERATED_BODY()
+
+public:
+	/** Default number of frames that should be buffered by a LiveLink source. */
+	UPROPERTY(config)
+	int32 DefaultSourceFrameBufferSize = 10;
+};
+
+
 class FArchive;
 class FProperty;
 class ULiveLinkSourceFactory;
@@ -148,6 +163,8 @@ class ULiveLinkSourceSettings : public UObject
 public:
 	GENERATED_BODY()
 
+	LIVELINKINTERFACE_API ULiveLinkSourceSettings();
+
 	/**
 	 * The the subject how to create the frame snapshot.
 	 * @note A client may evaluate manually the subject in a different mode by using EvaluateFrameAtWorldTime or EvaluateFrameAtSceneTime.
