LiveLinkHub - Change Parent Subject to Sync Subject in the UI

diff --git a/Engine/Source/Runtime/LiveLinkInterface/Public/LiveLinkSourceSettings.h b/Engine/Source/Runtime/LiveLinkInterface/Public/LiveLinkSourceSettings.h
--- a/Engine/Source/Runtime/LiveLinkInterface/Public/LiveLinkSourceSettings.h
+++ b/Engine/Source/Runtime/LiveLinkInterface/Public/LiveLinkSourceSettings.h
@@ -185,12 +185,12 @@ public:
 	TSubclassOf<ULiveLinkSourceFactory> Factory;
 
 	/** 
-	 * Parent subject used for resampling.
+	 * Which subject should be used as a synchronization source for this source.
 	 * If this is set, this source's subjects will only be rebroadcast when the parent subject receives data.
 	 * Additionally this source's subjects' timecode will match the parent's subject received timecode.
 	 * This can be useful for synchronizing a higher frequency source to a lower frequency one.
 	 */
-	UPROPERTY(EditAnywhere, Category = "Settings")
+	UPROPERTY(EditAnywhere, Category = "Settings", meta = (DisplayName="Sync Subject"))
 	FLiveLinkSubjectName ParentSubject;
 
 #if WITH_EDITORONLY_DATA
