LiveLink - Fix translators not working with rebroadcast data in virtual subjects

This caused by livelinkclient not being updated after a virtual subject was modified since it didn't call SetStaticDataAsRebroadcasted(false)

diff --git a/Engine/Plugins/Animation/LiveLink/Source/LiveLink/Private/Translator/LiveLinkAnimationRoleToTransform.cpp b/Engine/Plugins/Animation/LiveLink/Source/LiveLink/Private/Translator/LiveLinkAnimationRoleToTransform.cpp
--- a/Engine/Plugins/Animation/LiveLink/Source/LiveLink/Private/Translator/LiveLinkAnimationRoleToTransform.cpp
+++ b/Engine/Plugins/Animation/LiveLink/Source/LiveLink/Private/Translator/LiveLinkAnimationRoleToTransform.cpp
@@ -52,6 +52,8 @@ bool ULiveLinkAnimationRoleToTransform::FLiveLinkAnimationRoleToTransformWorker:
 	}
 
 	//Time to translate
+	TransformStaticData->PropertyNames.Append(SkeletonData->PropertyNames);
+
 	TransformFrameData->MetaData = FrameData->MetaData;
 	TransformFrameData->PropertyValues = FrameData->PropertyValues;
 	TransformFrameData->WorldTime = FrameData->WorldTime;
diff --git a/Engine/Source/Runtime/LiveLinkInterface/Private/LiveLinkVirtualSubject.cpp b/Engine/Source/Runtime/LiveLinkInterface/Private/LiveLinkVirtualSubject.cpp
--- a/Engine/Source/Runtime/LiveLinkInterface/Private/LiveLinkVirtualSubject.cpp
+++ b/Engine/Source/Runtime/LiveLinkInterface/Private/LiveLinkVirtualSubject.cpp
@@ -7,6 +7,7 @@
 
 #include UE_INLINE_GENERATED_CPP_BY_NAME(LiveLinkVirtualSubject)
 
+DEFINE_LOG_CATEGORY_STATIC(LogLiveLinkVirtualSubject, Log, All);
 
 void ULiveLinkVirtualSubject::Initialize(FLiveLinkSubjectKey InSubjectKey, TSubclassOf<ULiveLinkRole> InRole, ILiveLinkClient* InLiveLinkClient)
 {
@@ -115,4 +116,43 @@ void ULiveLinkVirtualSubject::InvalidateFrameData()
 	CurrentFrameSnapshot.FrameData.Reset();
 }
 
+bool ULiveLinkVirtualSubject::ValidateTranslators()
+{
+	UClass* RoleClass = Role.Get();
+	if (RoleClass == nullptr)
+	{
+		FrameTranslators.Reset();
+		return false;
+	}
+	else
+	{
+		for (int32 Index = 0; Index < FrameTranslators.Num(); ++Index)
+		{
+			if (ULiveLinkFrameTranslator* Translator = FrameTranslators[Index])
+			{
+				check(Translator->GetFromRole() != nullptr);
+				if (!RoleClass->IsChildOf(Translator->GetFromRole()))
+				{
+					UE_LOG(LogLiveLinkVirtualSubject, Warning, TEXT("Role '%s' is not supported by translator '%s'"), *RoleClass->GetName(), *Translator->GetName());
+					FrameTranslators[Index] = nullptr;
+					return false;
+				}
+			}
+		}
+	}
+
+	return true;
+}
+
+#if WITH_EDITOR
+void ULiveLinkVirtualSubject::PostEditChangeProperty(FPropertyChangedEvent& PropertyChangedEvent)
+{
+	Super::PostEditChangeProperty(PropertyChangedEvent);
 
+	if (PropertyChangedEvent.Property->GetFName() == GET_MEMBER_NAME_CHECKED(ULiveLinkVirtualSubject, FrameTranslators))
+	{
+		ValidateTranslators();
+		SetStaticDataAsRebroadcasted(false);
+	}
+}
+#endif
diff --git a/Engine/Source/Runtime/LiveLinkInterface/Public/LiveLinkVirtualSubject.h b/Engine/Source/Runtime/LiveLinkInterface/Public/LiveLinkVirtualSubject.h
--- a/Engine/Source/Runtime/LiveLinkInterface/Public/LiveLinkVirtualSubject.h
+++ b/Engine/Source/Runtime/LiveLinkInterface/Public/LiveLinkVirtualSubject.h
@@ -45,6 +45,7 @@ public:
 	virtual bool IsRebroadcasted() const override { return bRebroadcastSubject; }
 	virtual bool HasStaticDataBeenRebroadcasted() const override { return bHasStaticDataBeenRebroadcast; }
 	virtual void SetStaticDataAsRebroadcasted(const bool bInSent) override { bHasStaticDataBeenRebroadcast = bInSent; }
+
 protected:
 	virtual const FLiveLinkSubjectFrameData& GetFrameSnapshot() const override { return CurrentFrameSnapshot; }
 	//~ End ILiveLinkSubject Interface
@@ -56,6 +57,12 @@ protected:
 	LIVELINKINTERFACE_API bool HasValidFrameData() const;
 
 public:
+	//~ Begin UObject interface
+#if WITH_EDITOR
+	LIVELINKINTERFACE_API virtual void PostEditChangeProperty(struct FPropertyChangedEvent& PropertyChangedEvent) override;
+#endif
+	//~ End UObject interface
+
 	ILiveLinkClient* GetClient() const { return LiveLinkClient; }
 
 	/** Returns the live subjects associated with this virtual one */
@@ -130,6 +137,10 @@ protected:
 	 */
 	mutable FCriticalSection SnapshotAccessCriticalSection;
 
+private:
+	/** Validate that the translators on this subject match its role. */
+	bool ValidateTranslators();
+
 private:
 	TArray<ULiveLinkFrameTranslator::FWorkerSharedPtr> CurrentFrameTranslators;
 	/** Current subject remapper. */
