LiveLink timestamps for measuring latency

diff --git a/Engine/Source/Runtime/LiveLinkInterface/Public/LiveLinkTypes.h b/Engine/Source/Runtime/LiveLinkInterface/Public/LiveLinkTypes.h
--- a/Engine/Source/Runtime/LiveLinkInterface/Public/LiveLinkTypes.h
+++ b/Engine/Source/Runtime/LiveLinkInterface/Public/LiveLinkTypes.h
@@ -237,6 +237,12 @@ struct FLiveLinkBaseFrameData
 
 	/** Return the LiveLinkTime struct constructed from this frame data. */
 	FLiveLinkTime GetLiveLinkTime() const { return FLiveLinkTime(WorldTime.GetOffsettedTime(), MetaData.SceneTime); }
+
+#if WITH_EDITORONLY_DATA
+	/** This frame timestamps used for measuring latency */
+	UPROPERTY(VisibleAnywhere, Category = "LiveLink")
+	TMap<FName, double> Timestamps;
+#endif
 };
 
 
