[LiveLink]
- Marked ELiveLinkSubjectStatus as Blueprintable so it can be used with convenience switch statements in BP.

diff --git a/Engine/Source/Runtime/LiveLinkInterface/Public/ILiveLinkClient.h b/Engine/Source/Runtime/LiveLinkInterface/Public/ILiveLinkClient.h
--- a/Engine/Source/Runtime/LiveLinkInterface/Public/ILiveLinkClient.h
+++ b/Engine/Source/Runtime/LiveLinkInterface/Public/ILiveLinkClient.h
@@ -21,7 +21,7 @@ class ULiveLinkSourceSettings;
 
 
 /** Describes the state of a live link subject */
-UENUM()
+UENUM(BlueprintType)
 enum class ELiveLinkSubjectState : uint8
 {
 	/** The input is connected. */
