LiveLink - Fix double checkbox for RebroadcastSubject option

#rnx

diff --git a/Engine/Source/Runtime/LiveLinkInterface/Private/LiveLinkSubjectSettings.cpp b/Engine/Source/Runtime/LiveLinkInterface/Private/LiveLinkSubjectSettings.cpp
--- a/Engine/Source/Runtime/LiveLinkInterface/Private/LiveLinkSubjectSettings.cpp
+++ b/Engine/Source/Runtime/LiveLinkInterface/Private/LiveLinkSubjectSettings.cpp
@@ -18,7 +18,6 @@ ULiveLinkSubjectSettings::ULiveLinkSubjectSettings()
 	if (!HasAnyFlags(RF_ClassDefaultObject | RF_ArchetypeObject))
 	{
 		bRebroadcastSubject = GetDefault<ULiveLinkDefaultSubjectSettings>()->bRebroadcastSubjectsByDefault;
-		bAllowModifyingRebroadcast = GetDefault<ULiveLinkDefaultSubjectSettings>()->bAllowEditingRebroadcastProperty;
 	}
 }
 
diff --git a/Engine/Source/Runtime/LiveLinkInterface/Public/LiveLinkSubjectSettings.h b/Engine/Source/Runtime/LiveLinkInterface/Public/LiveLinkSubjectSettings.h
--- a/Engine/Source/Runtime/LiveLinkInterface/Public/LiveLinkSubjectSettings.h
+++ b/Engine/Source/Runtime/LiveLinkInterface/Public/LiveLinkSubjectSettings.h
@@ -36,13 +36,15 @@ public:
 	UPROPERTY(config)
 	bool bRebroadcastSubjectsByDefault = false;
 
+private:
 	/**
 	 * Whether a user should be able to edit the bRebroadcastSubject property.
 	 * Setting this to false in a target config will prevent a user from turning on or off the rebroadcast flag on a subject.
 	 *
 	 */
+	UE_DEPRECATED(5.6, "Not used anymore.")
 	UPROPERTY(config)
-	bool bAllowEditingRebroadcastProperty = true;
+	bool bAllowEditingRebroadcastProperty_DEPRECATED = true;
 };
 
 
@@ -94,13 +96,9 @@ public:
 	/** Last FrameRate estimated by the subject. If in Timecode mode, this will come directly from the QualifiedFrameTime. */
 	UPROPERTY(VisibleAnywhere, Category="LiveLink")
 	FFrameRate FrameRate;
-
-	/** Allows settings to dictate whether the rebroadcast flag is editable. */
-	UPROPERTY()
-	bool bAllowModifyingRebroadcast = true;
 	
 	/** If enabled, rebroadcast this subject */
-	UPROPERTY(EditAnywhere, Category = "LiveLink", meta=(EditCondition="bAllowModifyingRebroadcast"))
+	UPROPERTY(EditAnywhere, Category = "LiveLink")
     bool bRebroadcastSubject;
 
 	/** For sources created through LiveLinkHub, this contains the name of the original source for display purposes.*/
@@ -128,4 +126,9 @@ protected:
 private:
 	/** We need to keep track of the remapper when it's reset in order to restore the static data. */
 	TStrongObjectPtr<ULiveLinkSubjectRemapper> RemapperBeingReset;
+
+	/** Allows settings to dictate whether the rebroadcast flag is editable. */
+	UE_DEPRECATED(5.6, "Not used anymore.")
+	UPROPERTY()
+	bool bAllowModifyingRebroadcast_DEPRECATED = true;
 };
