[Live Link Interface]
- Add delegate to emit when a Live Link subject's bool bIsEnabled is modified. Required for Performance Capture workflow to respond to subjects being modified in the Live Link client panel.

diff --git a/Engine/Plugins/Animation/LiveLink/Source/LiveLink/Private/LiveLinkClient.cpp b/Engine/Plugins/Animation/LiveLink/Source/LiveLink/Private/LiveLinkClient.cpp
--- a/Engine/Plugins/Animation/LiveLink/Source/LiveLink/Private/LiveLinkClient.cpp
+++ b/Engine/Plugins/Animation/LiveLink/Source/LiveLink/Private/LiveLinkClient.cpp
@@ -2165,6 +2165,11 @@ FOnLiveLinkSubjectStateChanged& FLiveLinkClient::OnLiveLinkSubjectStateChanged()
 	return Collection->OnLiveLinkSubjectStateChanged();
 }
 
+FOnLiveLinkSubjectEnabledDelegate& FLiveLinkClient::OnLiveLinkSubjectEnabledChanged()
+{
+	return Collection->OnLiveLinkSubjectEnabledChanged();
+}
+
 #if WITH_EDITOR
 FOnLiveLinkSubjectEvaluated& FLiveLinkClient::OnLiveLinkSubjectEvaluated()
 {
diff --git a/Engine/Plugins/Animation/LiveLink/Source/LiveLink/Private/LiveLinkSourceCollection.cpp b/Engine/Plugins/Animation/LiveLink/Source/LiveLink/Private/LiveLinkSourceCollection.cpp
--- a/Engine/Plugins/Animation/LiveLink/Source/LiveLink/Private/LiveLinkSourceCollection.cpp
+++ b/Engine/Plugins/Animation/LiveLink/Source/LiveLink/Private/LiveLinkSourceCollection.cpp
@@ -318,11 +318,14 @@ void FLiveLinkSourceCollection::SetSubjectEnabled(FLiveLinkSubjectKey InSubjectK
 
 			NewEnabledItem->bEnabled = true;
 
+			BroadcastOnGameThread(OnLiveLinkSubjectEnabledChanged(), NewEnabledItem->Key, NewEnabledItem->bEnabled);
+			
 			for (FLiveLinkCollectionSubjectItem& SubjectItem : Subjects)
 			{
 				if (SubjectItem.bEnabled && SubjectItem.Key.SubjectName == InSubjectKey.SubjectName && !(SubjectItem.Key == InSubjectKey))
 				{
 					SubjectItem.bEnabled = false;
+					BroadcastOnGameThread(OnLiveLinkSubjectEnabledChanged(), SubjectItem.Key, SubjectItem.bEnabled);
 				}
 			}
 		}
@@ -335,6 +338,7 @@ void FLiveLinkSourceCollection::SetSubjectEnabled(FLiveLinkSubjectKey InSubjectK
 			if (SubjectItem.Key.SubjectName == InSubjectKey.SubjectName)
 			{
 				SubjectItem.bEnabled = false;
+				BroadcastOnGameThread(OnLiveLinkSubjectEnabledChanged(), SubjectItem.Key, SubjectItem.bEnabled);
 			}
 		}
 	}
diff --git a/Engine/Plugins/Animation/LiveLink/Source/LiveLink/Public/LiveLinkClient.h b/Engine/Plugins/Animation/LiveLink/Source/LiveLink/Public/LiveLinkClient.h
--- a/Engine/Plugins/Animation/LiveLink/Source/LiveLink/Public/LiveLinkClient.h
+++ b/Engine/Plugins/Animation/LiveLink/Source/LiveLink/Public/LiveLinkClient.h
@@ -113,6 +113,8 @@ public:
 	UE_API virtual FOnLiveLinkSubjectChangedDelegate& OnLiveLinkSubjectAdded() override;
 	UE_API virtual FOnLiveLinkSubjectChangedDelegate& OnLiveLinkSubjectRemoved() override;
 	UE_API virtual FOnLiveLinkSubjectStateChanged& OnLiveLinkSubjectStateChanged() override;
+	UE_API virtual FOnLiveLinkSubjectEnabledDelegate& OnLiveLinkSubjectEnabledChanged() override;
+	
 #if WITH_EDITOR
 	UE_API virtual FOnLiveLinkSubjectEvaluated& OnLiveLinkSubjectEvaluated() override;
 #endif
diff --git a/Engine/Plugins/Animation/LiveLink/Source/LiveLink/Public/LiveLinkSourceCollection.h b/Engine/Plugins/Animation/LiveLink/Source/LiveLink/Public/LiveLinkSourceCollection.h
--- a/Engine/Plugins/Animation/LiveLink/Source/LiveLink/Public/LiveLinkSourceCollection.h
+++ b/Engine/Plugins/Animation/LiveLink/Source/LiveLink/Public/LiveLinkSourceCollection.h
@@ -159,6 +159,7 @@ public:
 	FOnLiveLinkSubjectChangedDelegate& OnLiveLinkSubjectAdded() { return OnLiveLinkSubjectAddedDelegate; }
 	FOnLiveLinkSubjectChangedDelegate& OnLiveLinkSubjectRemoved() { return OnLiveLinkSubjectRemovedDelegate; }
 	FOnLiveLinkSubjectStateChanged& OnLiveLinkSubjectStateChanged() { return OnLiveLinkSubjectStateChangedDelegate; }
+	FOnLiveLinkSubjectEnabledDelegate& OnLiveLinkSubjectEnabledChanged() {return OnLiveLinkSubjectEnabledDelegate; } 
 
 	/** Utility method used to broadcast delegates on the game thread if this function is called on a different thread. */
 	template <typename DelegateType, typename... ArgTypes>
@@ -208,6 +209,9 @@ private:
 	/** Notify when a client subject's state has changed. (ie. it became stale) */
 	FOnLiveLinkSubjectStateChanged OnLiveLinkSubjectStateChangedDelegate;
 
+	/** Notify when a client subject's state of enabled has changed. */
+	FOnLiveLinkSubjectEnabledDelegate OnLiveLinkSubjectEnabledDelegate;
+
 	/** Lock to stop multiple threads accessing the Subjects from the collection at the same time */
 	mutable UE::FTransactionallySafeRecursiveMutex SubjectsLock;
 
diff --git a/Engine/Source/Runtime/LiveLinkInterface/Public/ILiveLinkClient.h b/Engine/Source/Runtime/LiveLinkInterface/Public/ILiveLinkClient.h
--- a/Engine/Source/Runtime/LiveLinkInterface/Public/ILiveLinkClient.h
+++ b/Engine/Source/Runtime/LiveLinkInterface/Public/ILiveLinkClient.h
@@ -49,7 +49,7 @@ DECLARE_TS_MULTICAST_DELEGATE_ThreeParams(FOnLiveLinkSubjectStaticDataAdded, FLi
 DECLARE_TS_MULTICAST_DELEGATE_ThreeParams(FOnLiveLinkSubjectFrameDataAdded, FLiveLinkSubjectKey /*InSubjectKey*/, TSubclassOf<ULiveLinkRole> /*SubjectRole*/, const FLiveLinkFrameDataStruct& /*InFrameData*/)
 DECLARE_TS_MULTICAST_DELEGATE_FiveParams(FOnLiveLinkSubjectEvaluated, FLiveLinkSubjectKey /*InSubjectKey*/, TSubclassOf<ULiveLinkRole> /*RequestedRole*/, const FLiveLinkTime& /*RequestedTime*/, bool /*bResult*/, const FLiveLinkTime& /*EvaluatedFrameTime*/)
 DECLARE_TS_MULTICAST_DELEGATE_TwoParams(FOnLiveLinkSubjectStateChanged, FLiveLinkSubjectKey /*InSubjectKey*/, ELiveLinkSubjectState /*NewState*/);
-
+DECLARE_TS_MULTICAST_DELEGATE_TwoParams(FOnLiveLinkSubjectEnabledDelegate, FLiveLinkSubjectKey /*SubjectKey*/, bool /*NewEnabled*/);
 
 /**
  * Interface for streaming and consuming data from external sources into UE4.
@@ -324,6 +324,9 @@ public:
 	/** Notify when a subject's state has changed. */
 	virtual FOnLiveLinkSubjectStateChanged& OnLiveLinkSubjectStateChanged() = 0;
 
+	/** Notify when a subject is enabled or disabled. */
+	virtual FOnLiveLinkSubjectEnabledDelegate& OnLiveLinkSubjectEnabledChanged() = 0;
+
 #if WITH_EDITOR
 	/** Notify the debug interface when a subject has been evaluated. Only available in editor and used for debugging purpose. */
 	virtual FOnLiveLinkSubjectEvaluated& OnLiveLinkSubjectEvaluated() = 0;
