LiveLink: Now defaults to 64 buffer size.

#rnx

#Synced-CL 51080076

diff --git a/Engine/Source/Runtime/LiveLinkInterface/Public/LiveLinkSourceSettings.h b/Engine/Source/Runtime/LiveLinkInterface/Public/LiveLinkSourceSettings.h
--- a/Engine/Source/Runtime/LiveLinkInterface/Public/LiveLinkSourceSettings.h
+++ b/Engine/Source/Runtime/LiveLinkInterface/Public/LiveLinkSourceSettings.h
@@ -28,7 +28,7 @@ class ULiveLinkDefaultSourceSettings : public UObject
 public:
 	/** Default number of frames that should be buffered by a LiveLink source. */
 	UPROPERTY(config)
-	int32 DefaultSourceFrameBufferSize = 10;
+	int32 DefaultSourceFrameBufferSize = 64;
 };
 
 
